Underscore primary plugin/module check-in.

diff --git a/Engine/Plugins/Runtime/Underscore/Source/Underscore/Private/Underscore.cpp b/Engine/Plugins/Runtime/Underscore/Source/Underscore/Private/Underscore.cpp
new file mode 100644
--- /dev/null
+++ b/Engine/Plugins/Runtime/Underscore/Source/Underscore/Private/Underscore.cpp
@@ -0,0 +1,22 @@
+// Copyright Epic Games, Inc. All Rights Reserved.
+
+#include "Underscore.h"
+
+#define LOCTEXT_NAMESPACE "FUnderscoreModule"
+
+void FUnderscoreModule::StartupModule()
+{
+	// This code will execute after your module is loaded into memory; the exact timing is specified in the .uplugin file per-module
+}
+
+void FUnderscoreModule::ShutdownModule()
+{
+	// This function may be called during shutdown to clean up your module.  For modules that support dynamic reloading,
+	// we call this function before unloading the module.
+}
+
+#undef LOCTEXT_NAMESPACE
+	
+DEFINE_LOG_CATEGORY(LogUnderscore);
+
+IMPLEMENT_MODULE(FUnderscoreModule, Underscore)
\ No newline at end of file
diff --git a/Engine/Plugins/Runtime/Underscore/Source/Underscore/Public/Underscore.h b/Engine/Plugins/Runtime/Underscore/Source/Underscore/Public/Underscore.h
new file mode 100644
--- /dev/null
+++ b/Engine/Plugins/Runtime/Underscore/Source/Underscore/Public/Underscore.h
@@ -0,0 +1,16 @@
+// Copyright Epic Games, Inc. All Rights Reserved.
+
+#pragma once
+
+#include "Modules/ModuleManager.h"
+
+DECLARE_LOG_CATEGORY_EXTERN(LogUnderscore, Log, All);
+
+class FUnderscoreModule : public IModuleInterface
+{
+public:
+
+	/** IModuleInterface implementation */
+	virtual void StartupModule() override;
+	virtual void ShutdownModule() override;
+};
diff --git a/Engine/Plugins/Runtime/Underscore/Source/Underscore/Underscore.Build.cs b/Engine/Plugins/Runtime/Underscore/Source/Underscore/Underscore.Build.cs
new file mode 100644
--- /dev/null
+++ b/Engine/Plugins/Runtime/Underscore/Source/Underscore/Underscore.Build.cs
@@ -0,0 +1,49 @@
+// Copyright Epic Games, Inc. All Rights Reserved.
+
+using UnrealBuildTool;
+
+public class Underscore : ModuleRules
+{
+	public Underscore(ReadOnlyTargetRules Target) : base(Target)
+	{
+		PCHUsage = ModuleRules.PCHUsageMode.UseExplicitOrSharedPCHs;
+
+		PublicIncludePaths.AddRange(
+			new string[] {
+			}
+			);
+				
+		PublicDependencyModuleNames.AddRange(
+			new string[]
+			{
+				"Core",
+				"AudioMixer",
+				"GameplayTags",
+				"StateTreeModule",
+				"AudioExtensions",
+
+			}
+			);
+			
+		
+		PrivateDependencyModuleNames.AddRange(
+			new string[]
+			{
+				"CoreUObject",
+				"Engine",
+				"Slate",
+				"SlateCore",
+			}
+			);
+
+		if (Target.bBuildEditor)
+		{
+			PrivateDependencyModuleNames.AddRange(
+				new string[]
+				{
+				"UnrealEd",
+				}
+				);
+		}
+	}
+}
diff --git a/Engine/Plugins/Runtime/Underscore/Source/UnderscoreEditor/Private/UnderscoreEditor.cpp b/Engine/Plugins/Runtime/Underscore/Source/UnderscoreEditor/Private/UnderscoreEditor.cpp
new file mode 100644
--- /dev/null
+++ b/Engine/Plugins/Runtime/Underscore/Source/UnderscoreEditor/Private/UnderscoreEditor.cpp
@@ -0,0 +1,34 @@
+// Copyright Epic Games, Inc. All Rights Reserved.
+
+#include "UnderscoreEditor.h"
+#include "AssetToolsModule.h"
+// #include "UnderscoreSection.h"
+// #include "AssetTypeActions_UnderscoreSection.h"
+
+void FUnderscoreEditorModule::StartupModule()
+{
+	// Register the audio editor asset type actions.
+ // 	IAssetTools& AssetTools = FModuleManager::LoadModuleChecked<FAssetToolsModule>("AssetTools").Get();
+
+//	UnderscoreSectionAssetTypeActions = MakeShareable(new FAssetTypeActions_UnderscoreSection);
+//	AssetTools.RegisterAssetTypeActions(UnderscoreSectionAssetTypeActions.ToSharedRef()); 
+
+}
+
+void FUnderscoreEditorModule::ShutdownModule()
+{
+	// This function may be called during shutdown to clean up your module.  For modules that support dynamic reloading,
+	// we call this function before unloading the module.
+	 
+/*	if (FModuleManager::Get().IsModuleLoaded("AssetTools"))
+	{
+		IAssetTools& AssetTools = FModuleManager::GetModuleChecked<FAssetToolsModule>("AssetTools").Get();
+		if (UnderscoreSectionAssetTypeActions.IsValid())
+		{
+			AssetTools.UnregisterAssetTypeActions(UnderscoreSectionAssetTypeActions.ToSharedRef());
+		}
+	}
+	*/
+}
+
+IMPLEMENT_MODULE(FUnderscoreEditorModule, UnderscoreEditor)
\ No newline at end of file
diff --git a/Engine/Plugins/Runtime/Underscore/Source/UnderscoreEditor/Public/UnderscoreEditor.h b/Engine/Plugins/Runtime/Underscore/Source/UnderscoreEditor/Public/UnderscoreEditor.h
new file mode 100644
--- /dev/null
+++ b/Engine/Plugins/Runtime/Underscore/Source/UnderscoreEditor/Public/UnderscoreEditor.h
@@ -0,0 +1,20 @@
+// Copyright Epic Games, Inc. All Rights Reserved.
+
+#pragma once
+
+#include "CoreMinimal.h"
+#include "Modules/ModuleInterface.h"
+#include "Modules/ModuleManager.h"
+#include "AssetToolsModule.h"
+
+class FUnderscoreEditorModule : public IModuleInterface
+{
+public:
+
+	/** IModuleInterface implementation */
+	virtual void StartupModule() override;
+	virtual void ShutdownModule() override;
+
+private:
+	TSharedPtr<IAssetTypeActions> UnderscoreSectionAssetTypeActions;
+};
diff --git a/Engine/Plugins/Runtime/Underscore/Source/UnderscoreEditor/UnderscoreEditor.Build.cs b/Engine/Plugins/Runtime/Underscore/Source/UnderscoreEditor/UnderscoreEditor.Build.cs
new file mode 100644
--- /dev/null
+++ b/Engine/Plugins/Runtime/Underscore/Source/UnderscoreEditor/UnderscoreEditor.Build.cs
@@ -0,0 +1,51 @@
+// Copyright Epic Games, Inc. All Rights Reserved.
+
+using UnrealBuildTool;
+
+public class UnderscoreEditor : ModuleRules
+{
+	public UnderscoreEditor(ReadOnlyTargetRules Target) : base(Target)
+	{
+		PrivatePCHHeaderFile = "Public/UnderscoreEditor.h";
+
+		PublicIncludePaths.AddRange(
+			new string[] {
+			}
+			);
+				
+		
+		PrivateIncludePaths.AddRange(
+			new string[] {
+			}
+			);
+			
+		
+		PublicDependencyModuleNames.AddRange(
+			new string[]
+			{
+				"Core",
+				"EditorFramework",
+				"UnrealEd",
+				"AssetTools",
+				"ContentBrowser",
+				"Underscore",
+				"SlateCore",
+
+			}
+			);
+			
+		
+		PrivateDependencyModuleNames.AddRange(
+			new string[]
+			{
+				"CoreUObject",
+				"Engine",
+				"Slate",
+				"Settings",
+				"UnrealEd",
+				"InputCore",
+				"EditorWidgets",
+			}
+			);
+	}
+}
diff --git a/Engine/Plugins/Runtime/Underscore/Underscore.uplugin b/Engine/Plugins/Runtime/Underscore/Underscore.uplugin
new file mode 100644
--- /dev/null
+++ b/Engine/Plugins/Runtime/Underscore/Underscore.uplugin
@@ -0,0 +1,36 @@
+{
+	"FileVersion": 3,
+	"Version": 1,
+	"VersionName": "0.2",
+	"FriendlyName": "Underscore",
+	"Description": "High-level interactive music system that simplifies dynamic playback while leaving room for some procedural playback behaviors.",
+	"Category": "Audio",
+	"CreatedBy": "Epic Games, Inc.",
+	"CreatedByURL": "",
+	"DocsURL": "",
+	"MarketplaceURL": "",
+	"SupportURL": "",
+	"EnabledByDefault": "false",
+	"CanContainContent": true,
+	"IsBetaVersion": false,
+	"IsExperimentalVersion": true,
+	"Installed": false,
+	"Modules": [
+		{
+			"Name": "Underscore",
+			"Type": "Runtime",
+			"LoadingPhase": "Default"
+		},
+		{
+			"Name": "UnderscoreEditor",
+			"Type": "Editor",
+			"LoadingPhase": "Default"
+		}
+	],
+	"Plugins": [
+		{
+			"Name": "StateTree",
+			"Enabled": true
+		}
+	]
+}
\ No newline at end of file
