[UnderscoreEditor] Removed private PCH.

diff --git a/Engine/Plugins/Runtime/Underscore/Source/UnderscoreEditor/UnderscoreEditor.Build.cs b/Engine/Plugins/Runtime/Underscore/Source/UnderscoreEditor/UnderscoreEditor.Build.cs
--- a/Engine/Plugins/Runtime/Underscore/Source/UnderscoreEditor/UnderscoreEditor.Build.cs
+++ b/Engine/Plugins/Runtime/Underscore/Source/UnderscoreEditor/UnderscoreEditor.Build.cs
@@ -6,8 +6,6 @@ public class UnderscoreEditor : ModuleRules
 {
 	public UnderscoreEditor(ReadOnlyTargetRules Target) : base(Target)
 	{
-		PrivatePCHHeaderFile = "Public/UnderscoreEditor.h";
-
 		PublicIncludePaths.AddRange(
 			new string[] {
 			}
